[PCG] Niagara Interop

* Add a new Interop plugin for PCG
* Add a new node: Write to NDC that does a fire and forget write from PCG to Niagara
* Add some traits for marshalling between PCG and Niagara types.

#rnx

diff --git a/Engine/Plugins/Experimental/PCGInterops/PCGNiagaraInterop/PCGNiagaraInterop.uplugin b/Engine/Plugins/Experimental/PCGInterops/PCGNiagaraInterop/PCGNiagaraInterop.uplugin
new file mode 100644
--- /dev/null
+++ b/Engine/Plugins/Experimental/PCGInterops/PCGNiagaraInterop/PCGNiagaraInterop.uplugin
@@ -0,0 +1,35 @@
+{
+	"FileVersion": 3,
+	"Version": 1,
+	"VersionName": "0.1",
+	"FriendlyName": "Procedural Content Generation Framework (PCG) Niagara Interop",
+	"Description": "Extra plugin for Procedural Content Generation Framework interacting with the Niagara system.",
+	"Category": "Editor",
+	"CreatedBy": "Epic Games, Inc.",
+	"CreatedByURL": "https://epicgames.com",
+	"DocsURL": "https://docs.unrealengine.com/latest/en-US/procedural-content-generation--framework-in-unreal-engine/",
+	"MarketplaceURL": "",
+	"SupportURL": "",
+	"EnabledByDefault": false,
+	"CanContainContent": true,
+	"IsBetaVersion": false,
+	"IsExperimentalVersion": true,
+	"Installed": false,
+	"Modules": [
+		{
+			"Name": "PCGNiagaraInterop",
+			"Type": "Runtime",
+			"LoadingPhase": "Default"
+		}
+	],
+	"Plugins": [
+		{
+			"Name": "PCG",
+			"Enabled": true
+		},
+		{
+			"Name": "Niagara",
+			"Enabled": true
+		}
+	]
+}
diff --git a/Engine/Plugins/Experimental/PCGInterops/PCGNiagaraInterop/Source/PCGNiagaraInterop/PCGNiagaraInterop.Build.cs b/Engine/Plugins/Experimental/PCGInterops/PCGNiagaraInterop/Source/PCGNiagaraInterop/PCGNiagaraInterop.Build.cs
new file mode 100644
--- /dev/null
+++ b/Engine/Plugins/Experimental/PCGInterops/PCGNiagaraInterop/Source/PCGNiagaraInterop/PCGNiagaraInterop.Build.cs
@@ -0,0 +1,22 @@
+// Copyright Epic Games, Inc. All Rights Reserved.
+
+namespace UnrealBuildTool.Rules
+{
+	public class PCGNiagaraInterop : ModuleRules
+	{
+		public PCGNiagaraInterop(ReadOnlyTargetRules Target) : base(Target)
+		{
+			PCHUsage = PCHUsageMode.UseExplicitOrSharedPCHs;
+
+			PublicDependencyModuleNames.AddRange(
+				new string[] {
+					"Core",
+					"CoreUObject",
+					"Engine",
+					"PCG",
+					"Niagara",
+				}
+			);
+		}
+	}
+}
diff --git a/Engine/Plugins/Experimental/PCGInterops/PCGNiagaraInterop/Source/PCGNiagaraInterop/Private/Elements/PCGWriteToNiagaraDataChannel.cpp b/Engine/Plugins/Experimental/PCGInterops/PCGNiagaraInterop/Source/PCGNiagaraInterop/Private/Elements/PCGWriteToNiagaraDataChannel.cpp
new file mode 100644
--- /dev/null
+++ b/Engine/Plugins/Experimental/PCGInterops/PCGNiagaraInterop/Source/PCGNiagaraInterop/Private/Elements/PCGWriteToNiagaraDataChannel.cpp
@@ -0,0 +1,208 @@
+// Copyright Epic Games, Inc. All Rights Reserved.
+
+#include "Elements/PCGWriteToNiagaraDataChannel.h"
+
+#include "Helpers/PCGAttributeNiagaraTraits.h"
+
+#include "PCGComponent.h"
+#include "PCGContext.h"
+#include "Data/PCGSpatialData.h"
+#include "Elements/Metadata/PCGMetadataElementCommon.h"
+#include "Metadata/Accessors/IPCGAttributeAccessor.h"
+#include "Metadata/Accessors/PCGAttributeAccessorHelpers.h"
+#include "Metadata/Accessors/PCGAttributeAccessorKeys.h"
+#include "Utils/PCGLogErrors.h"
+
+#include "NiagaraDataChannel.h"
+#include "NiagaraDataChannelAccessor.h"
+#include "NiagaraDataChannelPublic.h"
+
+#include "Engine/World.h"
+#include "GameFramework/Actor.h"
+
+#define LOCTEXT_NAMESPACE "PCGWriteToNDCElement"
+
+#if WITH_EDITOR
+FName UPCGWriteToNiagaraDataChannelSettings::GetDefaultNodeName() const
+{
+	return FName(TEXT("WriteToNDC"));
+}
+
+FText UPCGWriteToNiagaraDataChannelSettings::GetDefaultNodeTitle() const
+{
+	return LOCTEXT("NodeTitle", "Write To NDC");
+}
+#endif // WITH_EDITOR
+
+FPCGElementPtr UPCGWriteToNiagaraDataChannelSettings::CreateElement() const
+{
+	return MakeShared<FPCGWriteToNiagaraDataChannelElement>();
+}
+
+TArray<FPCGPinProperties> UPCGWriteToNiagaraDataChannelSettings::InputPinProperties() const
+{
+	TArray<FPCGPinProperties> Properties;
+	Properties.Emplace_GetRef(PCGPinConstants::DefaultInputLabel, EPCGDataType::Any).SetRequiredPin();
+	return Properties;
+}
+
+TArray<FPCGPinProperties> UPCGWriteToNiagaraDataChannelSettings::OutputPinProperties() const
+{ 
+	TArray<FPCGPinProperties> Properties;
+	Properties.Emplace(PCGPinConstants::DefaultOutputLabel, EPCGDataType::Any);
+	return Properties;
+}
+
+bool FPCGWriteToNiagaraDataChannelElement::PrepareDataInternal(FPCGContext* InContext) const
+{
+	TRACE_CPUPROFILER_EVENT_SCOPE(FPCGWriteToNiagaraDataChannelElement::Execute);
+
+	FPCGWriteToNiagaraDataChannelContext* Context = reinterpret_cast<FPCGWriteToNiagaraDataChannelContext*>(InContext);
+	check(InContext);
+
+	const UPCGWriteToNiagaraDataChannelSettings* Settings = InContext->GetInputSettings<UPCGWriteToNiagaraDataChannelSettings>();
+	check(Settings);
+	
+	if (Settings->DataChannel.IsNull())
+	{
+		return true;
+	}
+
+	if (!Context->WasLoadRequested())
+	{
+		return Context->RequestResourceLoad(Context, {Settings->DataChannel.ToSoftObjectPath()}, !Settings->bSynchronousLoad);
+	}
+	else
+	{
+		return true;
+	}
+}
+
+bool FPCGWriteToNiagaraDataChannelElement::ExecuteInternal(FPCGContext* InContext) const
+{
+	TRACE_CPUPROFILER_EVENT_SCOPE(FPCGWriteToNiagaraDataChannelElement::Execute);
+
+	check(InContext);
+
+	const UPCGWriteToNiagaraDataChannelSettings* Settings = InContext->GetInputSettings<UPCGWriteToNiagaraDataChannelSettings>();
+	check(Settings);
+
+	// Passthrough input
+	InContext->OutputData = InContext->InputData;
+
+	UNiagaraDataChannelAsset* DataChannelAsset = Settings->DataChannel.Get();
+	UNiagaraDataChannel* DataChannel = DataChannelAsset ? DataChannelAsset->Get() : nullptr;
+	if (!DataChannel)
+	{
+		return true;
+	}
+
+	TArray<FName, TInlineAllocator<16>> MatchedNiagaraVars;
+	MatchedNiagaraVars.Reserve(Settings->NiagaraVariablesPCGAttributeMapping.Num());
+	bool bUnmatchedVarsWereWarned = false;
+
+	int InputNum = 0;
+
+	UWorld* World = InContext->SourceComponent.IsValid() ? InContext->SourceComponent->GetWorld() : nullptr;
+	if (!World)
+	{
+		return true;
+	}
+
+	for (const FPCGTaggedData& Input : InContext->InputData.GetInputsByPin(PCGPinConstants::DefaultInputLabel))
+	{
+		InputNum++;
+		int32 Count = -1;
+
+		TArray<TTuple<const FNiagaraVariableBase&, TUniquePtr<const IPCGAttributeAccessor>, TUniquePtr<const IPCGAttributeAccessorKeys>>> AccessorAndKeys;
+		for (const FNiagaraDataChannelVariable& NiagaraVar : DataChannel->GetVariables())
+		{
+			const FName NiagaraVarName = NiagaraVar.GetName();
+			if (const FPCGAttributePropertyInputSelector* Selector = Settings->NiagaraVariablesPCGAttributeMapping.Find(NiagaraVarName))
+			{
+				if (!bUnmatchedVarsWereWarned)
+				{
+					MatchedNiagaraVars.Add(NiagaraVarName);
+				}
+
+				const FPCGAttributePropertyInputSelector FixedSelector = Selector->CopyAndFixLast(Input.Data);
+
+				TUniquePtr<const IPCGAttributeAccessor> Accessor = PCGAttributeAccessorHelpers::CreateConstAccessor(Input.Data, FixedSelector);
+				TUniquePtr<const IPCGAttributeAccessorKeys> Keys = PCGAttributeAccessorHelpers::CreateConstKeys(Input.Data, FixedSelector);
+
+				if (!Accessor || !Keys)
+				{
+					PCGLog::Metadata::LogFailToCreateAccessorError(FixedSelector, InContext);
+					continue;
+				}
+
+				// Verify that types are compatible
+				if (!PCGAttributeNiagaraTraits::AreTypesCompatible(Accessor->GetUnderlyingType(), NiagaraVar, /*bPCGToNiagara=*/true))
+				{
+					PCGLog::LogErrorOnGraph(FText::Format(LOCTEXT("MismatchNiagaraVarType", "Niagara variable {0} is not compatible with attribute {1} ({2})."), 
+						FText::FromName(NiagaraVarName), FixedSelector.GetDisplayText(), PCG::Private::GetTypeNameText(Accessor->GetUnderlyingType())), InContext);
+					continue;
+				}
+
+				Count = FMath::Max(Count, Keys->GetNum());
+				AccessorAndKeys.Emplace(NiagaraVar, std::move(Accessor), std::move(Keys));
+			}
+		}
+
+		if (!bUnmatchedVarsWereWarned && MatchedNiagaraVars.Num() != Settings->NiagaraVariablesPCGAttributeMapping.Num())
+		{
+			for (const auto& [NiagaraVarName, Selector] : Settings->NiagaraVariablesPCGAttributeMapping)
+			{
+				if (!MatchedNiagaraVars.Contains(NiagaraVarName))
+				{
+					PCGLog::LogErrorOnGraph(FText::Format(LOCTEXT("MismatchNiagaraVarName", "Niagara variable {0} was not found in the data channel."), FText::FromName(NiagaraVarName)), InContext);
+				}
+			}
+		}
+
+		bUnmatchedVarsWereWarned = true;
+
+		if (AccessorAndKeys.IsEmpty() || Count <= 0)
+		{
+			continue;
+		}
+
+		FNiagaraDataChannelSearchParameters SearchParameters;
+		AActor* TargetActor = InContext->GetTargetActor(nullptr);
+		SearchParameters.OwningComponent = TargetActor ? TargetActor->GetRootComponent() : nullptr;
+
+		const FString DebugSource = FString::Format(TEXT("PCGWriteToNiagaraChannel - {0} - Input {1}"), { InContext->Node ? InContext->Node->GetName() : TEXT("Unknown node"), FString::FromInt(InputNum - 1) });
+
+		UNiagaraDataChannelWriter* NiagaraWriter = UNiagaraDataChannelLibrary::CreateDataChannelWriter(World, DataChannel, std::move(SearchParameters), Count, /*bVisibleToGame=*/ true, /*bVisibleToCPU=*/ true, /*bVisibleToGPU=*/ false, DebugSource);
+
+		if (!NiagaraWriter)
+		{
+			PCGLog::LogErrorOnGraph(LOCTEXT("FailToCreateWriter", "Failed to create Niagara Data Channel Writer"), InContext);
+			continue;
+		}
+
+		for (const auto& [NiagaraVar, Accessor, Keys] : AccessorAndKeys)
+		{
+			PCGAttributeNiagaraTraits::CallbackWithNiagaraType(NiagaraVar, [&NiagaraVar, &Accessor, &Keys, NiagaraWriter]<typename NiagaraType>(NiagaraType)
+			{
+				// PCG doesn't support linear color natively, so read it from a FVector4.
+				using PCGType = std::conditional_t<std::is_same_v<NiagaraType, FLinearColor>, FVector4, NiagaraType>;
+				return PCGMetadataElementCommon::ApplyOnAccessor<PCGType>(*Keys, *Accessor, [NiagaraWriter, &NiagaraVar](const PCGType& InValue, int32 Index)
+					{
+						if constexpr (std::is_same_v<NiagaraType, FLinearColor>)
+						{
+							NiagaraWriter->WriteData(NiagaraVar, Index, FLinearColor{ InValue });
+						}
+						else
+						{
+							NiagaraWriter->WriteData(NiagaraVar, Index, InValue);
+						}
+					}, EPCGAttributeAccessorFlags::AllowBroadcastAndConstructible);
+			});
+		}
+	}
+
+	return true;
+}
+
+#undef LOCTEXT_NAMESPACE
diff --git a/Engine/Plugins/Experimental/PCGInterops/PCGNiagaraInterop/Source/PCGNiagaraInterop/Private/Helpers/PCGAttributeNiagaraTraits.cpp b/Engine/Plugins/Experimental/PCGInterops/PCGNiagaraInterop/Source/PCGNiagaraInterop/Private/Helpers/PCGAttributeNiagaraTraits.cpp
new file mode 100644
--- /dev/null
+++ b/Engine/Plugins/Experimental/PCGInterops/PCGNiagaraInterop/Source/PCGNiagaraInterop/Private/Helpers/PCGAttributeNiagaraTraits.cpp
@@ -0,0 +1,22 @@
+// Copyright Epic Games, Inc. All Rights Reserved.
+
+#include "Helpers/PCGAttributeNiagaraTraits.h"
+
+bool PCGAttributeNiagaraTraits::AreTypesCompatible(uint16 PCGType, const FNiagaraVariableBase& NiagaraVar, bool bPCGToNiagara)
+{
+	return CallbackWithNiagaraType(NiagaraVar, [PCGType, bPCGToNiagara]<typename T>(T) -> bool
+	{
+		if constexpr (std::is_same_v<T, FNiagaraSpawnInfo>)
+		{
+			// Not supported at the moment
+			return false;
+		}
+		else
+		{
+			using NiagaraType = std::conditional_t<std::is_same_v<T, FLinearColor>, FVector4, T>;
+			return bPCGToNiagara
+				? PCG::Private::IsBroadcastableOrConstructible(PCGType, PCG::Private::MetadataTypes<NiagaraType>::Id)
+				: PCG::Private::IsBroadcastableOrConstructible(PCG::Private::MetadataTypes<NiagaraType>::Id, PCGType);
+		}
+	});
+}
\ No newline at end of file
diff --git a/Engine/Plugins/Experimental/PCGInterops/PCGNiagaraInterop/Source/PCGNiagaraInterop/Private/PCGNiagaraInteropModule.cpp b/Engine/Plugins/Experimental/PCGInterops/PCGNiagaraInterop/Source/PCGNiagaraInterop/Private/PCGNiagaraInteropModule.cpp
new file mode 100644
--- /dev/null
+++ b/Engine/Plugins/Experimental/PCGInterops/PCGNiagaraInterop/Source/PCGNiagaraInterop/Private/PCGNiagaraInteropModule.cpp
@@ -0,0 +1,9 @@
+// Copyright Epic Games, Inc. All Rights Reserved.
+
+#include "PCGNiagaraInteropModule.h"
+
+#include "Modules/ModuleManager.h"
+
+IMPLEMENT_MODULE(FPCGNiagaraInteropModule, PCGNiagaraInterop);
+
+DEFINE_LOG_CATEGORY(LogPCGNiagaraInterop);
diff --git a/Engine/Plugins/Experimental/PCGInterops/PCGNiagaraInterop/Source/PCGNiagaraInterop/Public/Elements/PCGWriteToNiagaraDataChannel.h b/Engine/Plugins/Experimental/PCGInterops/PCGNiagaraInterop/Source/PCGNiagaraInterop/Public/Elements/PCGWriteToNiagaraDataChannel.h
new file mode 100644
--- /dev/null
+++ b/Engine/Plugins/Experimental/PCGInterops/PCGNiagaraInterop/Source/PCGNiagaraInterop/Public/Elements/PCGWriteToNiagaraDataChannel.h
@@ -0,0 +1,58 @@
+// Copyright Epic Games, Inc. All Rights Reserved.
+
+#pragma once
+
+#include "PCGContext.h"
+#include "PCGSettings.h"
+#include "Async/PCGAsyncLoadingContext.h"
+#include "Metadata/PCGAttributePropertySelector.h"
+
+#include "PCGWriteToNiagaraDataChannel.generated.h"
+
+class UNiagaraDataChannelAsset;
+
+/**
+* Allow writing attributes to a Niagara Data Channel.
+*/
+UCLASS(MinimalAPI, BlueprintType, ClassGroup = (Procedural))
+class UPCGWriteToNiagaraDataChannelSettings : public UPCGSettings
+{
+	GENERATED_BODY()
+
+public:
+	//~Begin UPCGSettings interface
+#if WITH_EDITOR
+	virtual FName GetDefaultNodeName() const override;
+	virtual FText GetDefaultNodeTitle() const override;
+	virtual EPCGSettingsType GetType() const override { return EPCGSettingsType::Generic; }
+#endif
+
+protected:
+	virtual TArray<FPCGPinProperties> InputPinProperties() const override;
+	virtual TArray<FPCGPinProperties> OutputPinProperties() const override;
+	virtual FPCGElementPtr CreateElement() const override;
+	//~End UPCGSettings interface
+
+public:
+	UPROPERTY(BlueprintReadWrite, EditAnywhere, Category = "Settings", meta = (PCG_Overridable))
+	TSoftObjectPtr<UNiagaraDataChannelAsset> DataChannel;
+
+	UPROPERTY(BlueprintReadWrite, EditAnywhere, Category = "Settings")
+	TMap<FName, FPCGAttributePropertyInputSelector> NiagaraVariablesPCGAttributeMapping;
+
+	UPROPERTY(BlueprintReadWrite, EditAnywhere, Category = "Settings|Debug")
+	bool bSynchronousLoad = false;
+};
+
+struct FPCGWriteToNiagaraDataChannelContext : public FPCGContext, public IPCGAsyncLoadingContext {};
+
+class FPCGWriteToNiagaraDataChannelElement : public IPCGElementWithCustomContext<FPCGWriteToNiagaraDataChannelContext>
+{
+public:
+	virtual bool IsCacheable(const UPCGSettings* InSettings) const override { return false; }
+	virtual bool CanExecuteOnlyOnMainThread(FPCGContext* Context) const override { return true; }
+protected:
+	virtual bool PrepareDataInternal(FPCGContext* InContext) const override;
+	virtual bool ExecuteInternal(FPCGContext* InContext) const override;
+};
+
diff --git a/Engine/Plugins/Experimental/PCGInterops/PCGNiagaraInterop/Source/PCGNiagaraInterop/Public/Helpers/PCGAttributeNiagaraTraits.h b/Engine/Plugins/Experimental/PCGInterops/PCGNiagaraInterop/Source/PCGNiagaraInterop/Public/Helpers/PCGAttributeNiagaraTraits.h
new file mode 100644
--- /dev/null
+++ b/Engine/Plugins/Experimental/PCGInterops/PCGNiagaraInterop/Source/PCGNiagaraInterop/Public/Helpers/PCGAttributeNiagaraTraits.h
@@ -0,0 +1,68 @@
+// Copyright Epic Games, Inc. All Rights Reserved.
+
+#pragma once
+
+#include "Metadata/PCGMetadataAttributeTraits.h"
+#include "Metadata/PCGMetadataAttribute.h"
+
+#include "NiagaraTypes.h"
+
+namespace PCGAttributeNiagaraTraits
+{
+	bool AreTypesCompatible(uint16 PCGType, const FNiagaraVariableBase& NiagaraVar, bool bPCGToNiagara);
+
+	template <typename Func, typename  ...Args>
+	decltype(auto) CallbackWithNiagaraType(const FNiagaraVariableBase& NiagaraVar, Func Callback, Args&& ...InArgs)
+	{
+		using ReturnType = decltype(Callback(double{}, std::forward<Args>(InArgs)...));
+		const FNiagaraTypeDefinition& NiagaraType = NiagaraVar.GetType();
+
+		if (NiagaraType == FNiagaraTypeHelper::GetVector2DDef())
+		{
+			return Callback(FVector2D{}, std::forward<Args>(InArgs)...);
+		}
+		else if (NiagaraType == FNiagaraTypeHelper::GetVectorDef() || NiagaraType == FNiagaraTypeDefinition::GetPositionDef())
+		{
+			return Callback(FVector{}, std::forward<Args>(InArgs)...);
+		}
+		else if (NiagaraType == FNiagaraTypeHelper::GetVector4Def())
+		{
+			return Callback(FVector4{}, std::forward<Args>(InArgs)...);
+		}
+		else if (NiagaraType == FNiagaraTypeDefinition::GetColorDef())
+		{
+			return Callback(FLinearColor{}, std::forward<Args>(InArgs)...);
+		}
+		else if (NiagaraType == FNiagaraTypeHelper::GetQuatDef())
+		{
+			return Callback(FQuat{}, std::forward<Args>(InArgs)...);
+		}
+		else if (NiagaraType == FNiagaraTypeHelper::GetDoubleDef())
+		{
+			return Callback(double{}, std::forward<Args>(InArgs)...);
+		}
+		else if (NiagaraType == FNiagaraTypeDefinition::GetIntDef() || NiagaraType == FNiagaraTypeDefinition::GetIDDef())
+		{
+			return Callback(int64{}, std::forward<Args>(InArgs)...);
+		}
+		else if (NiagaraType == FNiagaraTypeDefinition::GetBoolDef())
+		{
+			return Callback(bool{}, std::forward<Args>(InArgs)...);
+		}
+		else if (NiagaraType == FNiagaraTypeDefinition(FNiagaraSpawnInfo::StaticStruct()))
+		{
+			return Callback(FNiagaraSpawnInfo{}, std::forward<Args>(InArgs)...);
+		}
+		else
+		{
+			if constexpr (std::is_same_v<ReturnType, void>)
+			{
+				return;
+			}
+			else
+			{
+				return ReturnType{};
+			}
+		}
+	}
+}
\ No newline at end of file
diff --git a/Engine/Plugins/Experimental/PCGInterops/PCGNiagaraInterop/Source/PCGNiagaraInterop/Public/PCGNiagaraInteropModule.h b/Engine/Plugins/Experimental/PCGInterops/PCGNiagaraInterop/Source/PCGNiagaraInterop/Public/PCGNiagaraInteropModule.h
new file mode 100644
--- /dev/null
+++ b/Engine/Plugins/Experimental/PCGInterops/PCGNiagaraInterop/Source/PCGNiagaraInterop/Public/PCGNiagaraInteropModule.h
@@ -0,0 +1,16 @@
+// Copyright Epic Games, Inc. All Rights Reserved.
+
+#pragma once
+
+#include "Modules/ModuleInterface.h"
+#include "Stats/Stats.h"
+
+DECLARE_LOG_CATEGORY_EXTERN(LogPCGNiagaraInterop, Log, All);
+
+class FPCGNiagaraInteropModule final : public IModuleInterface
+{
+public:
+	// ~IModuleInterface implementation
+	virtual bool SupportsDynamicReloading() override { return true; };
+	// ~End IModuleInterface implementation
+};
