[PCG] Fix build on android

#rnx

diff --git a/Engine/Plugins/Experimental/PCGInterops/PCGNiagaraInterop/Source/PCGNiagaraInterop/Private/Elements/PCGWriteToNiagaraDataChannel.cpp b/Engine/Plugins/Experimental/PCGInterops/PCGNiagaraInterop/Source/PCGNiagaraInterop/Private/Elements/PCGWriteToNiagaraDataChannel.cpp
--- a/Engine/Plugins/Experimental/PCGInterops/PCGNiagaraInterop/Source/PCGNiagaraInterop/Private/Elements/PCGWriteToNiagaraDataChannel.cpp
+++ b/Engine/Plugins/Experimental/PCGInterops/PCGNiagaraInterop/Source/PCGNiagaraInterop/Private/Elements/PCGWriteToNiagaraDataChannel.cpp
@@ -114,7 +114,8 @@ bool FPCGWriteToNiagaraDataChannelElement::ExecuteInternal(FPCGContext* InContex
 		InputNum++;
 		int32 Count = -1;
 
-		TArray<TTuple<const FNiagaraVariableBase&, TUniquePtr<const IPCGAttributeAccessor>, TUniquePtr<const IPCGAttributeAccessorKeys>>> AccessorAndKeys;
+		using FAccessorAndKeysTuple = TTuple<const FNiagaraVariableBase&, TUniquePtr<const IPCGAttributeAccessor>, TUniquePtr<const IPCGAttributeAccessorKeys>>;
+		TArray<FAccessorAndKeysTuple> AccessorAndKeys;
 		for (const FNiagaraDataChannelVariable& NiagaraVar : DataChannel->GetVariables())
 		{
 			const FName NiagaraVarName = NiagaraVar.GetName();
@@ -181,8 +182,12 @@ bool FPCGWriteToNiagaraDataChannelElement::ExecuteInternal(FPCGContext* InContex
 			continue;
 		}
 
-		for (const auto& [NiagaraVar, Accessor, Keys] : AccessorAndKeys)
+		for (const FAccessorAndKeysTuple& It : AccessorAndKeys)
 		{
+			const FNiagaraVariableBase& NiagaraVar = It.Get<0>();
+			const TUniquePtr<const IPCGAttributeAccessor>& Accessor = It.Get<1>();
+			const TUniquePtr<const IPCGAttributeAccessorKeys>& Keys = It.Get<2>();
+			
 			PCGAttributeNiagaraTraits::CallbackWithNiagaraType(NiagaraVar, [&NiagaraVar, &Accessor, &Keys, NiagaraWriter]<typename NiagaraType>(NiagaraType)
 			{
 				// PCG doesn't support linear color natively, so read it from a FVector4.
