[PCG] Typo on Niagara Interop node

#rnx

diff --git a/Engine/Plugins/Experimental/PCGInterops/PCGNiagaraInterop/Source/PCGNiagaraInterop/Private/Elements/PCGWriteToNiagaraDataChannel.cpp b/Engine/Plugins/Experimental/PCGInterops/PCGNiagaraInterop/Source/PCGNiagaraInterop/Private/Elements/PCGWriteToNiagaraDataChannel.cpp
--- a/Engine/Plugins/Experimental/PCGInterops/PCGNiagaraInterop/Source/PCGNiagaraInterop/Private/Elements/PCGWriteToNiagaraDataChannel.cpp
+++ b/Engine/Plugins/Experimental/PCGInterops/PCGNiagaraInterop/Source/PCGNiagaraInterop/Private/Elements/PCGWriteToNiagaraDataChannel.cpp
@@ -25,12 +25,12 @@
 #if WITH_EDITOR
 FName UPCGWriteToNiagaraDataChannelSettings::GetDefaultNodeName() const
 {
-	return FName(TEXT("WriteToNDC"));
+	return FName(TEXT("WriteToNiagaraDataChannel"));
 }
 
 FText UPCGWriteToNiagaraDataChannelSettings::GetDefaultNodeTitle() const
 {
-	return LOCTEXT("NodeTitle", "Write To NDC");
+	return LOCTEXT("NodeTitle", "Write To Niagara Data Channel");
 }
 #endif // WITH_EDITOR
 
