[PCG] Fix data marshalling between PCG and Niagara

* Ints and IDs were not working correctly

#rnx

diff --git a/Engine/Plugins/Experimental/PCGInterops/PCGNiagaraInterop/Source/PCGNiagaraInterop/Public/Helpers/PCGAttributeNiagaraTraits.h b/Engine/Plugins/Experimental/PCGInterops/PCGNiagaraInterop/Source/PCGNiagaraInterop/Public/Helpers/PCGAttributeNiagaraTraits.h
--- a/Engine/Plugins/Experimental/PCGInterops/PCGNiagaraInterop/Source/PCGNiagaraInterop/Public/Helpers/PCGAttributeNiagaraTraits.h
+++ b/Engine/Plugins/Experimental/PCGInterops/PCGNiagaraInterop/Source/PCGNiagaraInterop/Public/Helpers/PCGAttributeNiagaraTraits.h
@@ -41,9 +41,13 @@ namespace PCGAttributeNiagaraTraits
 		{
 			return Callback(double{}, std::forward<Args>(InArgs)...);
 		}
-		else if (NiagaraType == FNiagaraTypeDefinition::GetIntDef() || NiagaraType == FNiagaraTypeDefinition::GetIDDef())
+		else if (NiagaraType == FNiagaraTypeDefinition::GetIntDef())
 		{
-			return Callback(int64{}, std::forward<Args>(InArgs)...);
+			return Callback(int32{}, std::forward<Args>(InArgs)...);
+		}
+		else if (NiagaraType == FNiagaraTypeDefinition::GetIDDef())
+		{
+			return Callback(FNiagaraID{}, std::forward<Args>(InArgs)...);
 		}
 		else if (NiagaraType == FNiagaraTypeDefinition::GetBoolDef())
 		{
