[PCG] Fix bool attribute in Write to Niagara Data Channel

#rnx

diff --git a/Engine/Plugins/Experimental/PCGInterops/PCGNiagaraInterop/Source/PCGNiagaraInterop/Private/Elements/PCGWriteToNiagaraDataChannel.cpp b/Engine/Plugins/Experimental/PCGInterops/PCGNiagaraInterop/Source/PCGNiagaraInterop/Private/Elements/PCGWriteToNiagaraDataChannel.cpp
--- a/Engine/Plugins/Experimental/PCGInterops/PCGNiagaraInterop/Source/PCGNiagaraInterop/Private/Elements/PCGWriteToNiagaraDataChannel.cpp
+++ b/Engine/Plugins/Experimental/PCGInterops/PCGNiagaraInterop/Source/PCGNiagaraInterop/Private/Elements/PCGWriteToNiagaraDataChannel.cpp
@@ -198,6 +198,11 @@ bool FPCGWriteToNiagaraDataChannelElement::ExecuteInternal(FPCGContext* InContex
 						{
 							NiagaraWriter->WriteData(NiagaraVar, Index, FLinearColor{ InValue });
 						}
+						else if constexpr (std::is_same_v<NiagaraType, bool>)
+						{
+							// Need to use the FNiagaraBool struct
+							NiagaraWriter->WriteData(NiagaraVar, Index, FNiagaraBool(InValue));
+						}
 						else
 						{
 							NiagaraWriter->WriteData(NiagaraVar, Index, InValue);
